Reimplement CL 3007370 from 4.12 release. Change CRC config setting merged from Paragon. Allows user to close CRC without sending.
[UE-31812] - CrashReporter Client has no options to close without sending

diff --git a/Engine/Programs/CrashReportClient/Config/DefaultEngine.ini b/Engine/Programs/CrashReportClient/Config/DefaultEngine.ini
--- a/Engine/Programs/CrashReportClient/Config/DefaultEngine.ini
+++ b/Engine/Programs/CrashReportClient/Config/DefaultEngine.ini
@@ -4,7 +4,7 @@ DataRouterUrl="https://datarouter-public-service-prod06.ol.epicgames.com/datarou
 bAllowToBeContacted=true
 bSendLogFile=true
 bHideLogFilesOption=false
-bIsAllowedToCloseWithoutSending=false
+bIsAllowedToCloseWithoutSending=true
 CanSendWhenUIFailedToInitialize=true
 UIInitRetryCount=10
 UIInitRetryInterval=2.0
\ No newline at end of file
