Move Epic data router from public CRC config

Instead moved into Engine config and converted into a compile time constant compiled into CRCE.

diff --git a/Engine/Programs/CrashReportClient/Config/DefaultEngine.ini b/Engine/Programs/CrashReportClient/Config/DefaultEngine.ini
--- a/Engine/Programs/CrashReportClient/Config/DefaultEngine.ini
+++ b/Engine/Programs/CrashReportClient/Config/DefaultEngine.ini
@@ -1,8 +1,5 @@
 [CrashReportClient]
 CrashReportClientVersion=1.0
-DataRouterUrl="https://datarouter.ol.epicgames.com/datarouter/api/v1/public/data"
-bAllowToBeContacted=true
-bSendLogFile=true
 CanSendWhenUIFailedToInitialize=true
 UIInitRetryCount=10
 UIInitRetryInterval=2.0
