Fix default values for attaching logs and contact options in crashes.

diff --git a/Engine/Programs/CrashReportClient/Config/DefaultEngine.ini b/Engine/Programs/CrashReportClient/Config/DefaultEngine.ini
--- a/Engine/Programs/CrashReportClient/Config/DefaultEngine.ini
+++ b/Engine/Programs/CrashReportClient/Config/DefaultEngine.ini
@@ -1,5 +1,7 @@
 [CrashReportClient]
 CrashReportClientVersion=1.0
+bAllowToBeContacted=true
+bSendLogFile=true
 CanSendWhenUIFailedToInitialize=true
 UIInitRetryCount=10
 UIInitRetryInterval=2.0
